- Fixed xcodecheckout script to not point to /usr/local/bin/p4, but just run p4 from wherever it is, after sourcing zshrc and zprofile to help look for p4

#changelist validated
#virtualized

diff --git a/Engine/Extras/Xcode/xcodecheckout.sh b/Engine/Extras/Xcode/xcodecheckout.sh
--- a/Engine/Extras/Xcode/xcodecheckout.sh
+++ b/Engine/Extras/Xcode/xcodecheckout.sh
@@ -1,9 +1,42 @@
-#!/bin/bash
+#!/bin/zsh
 
 
 # This script can be specified in Xcode's behaviors to checkout files that are locked for editing.
 # It assumes the P4 commandline tools are installed, and the active p4config points to a workspace
 # that maps to the path of the file
 
- /usr/local/bin/p4 edit ${XcodeAlertAffectedPaths}
+# Look in various standard places
+if [ -x "$(command -v p4)" ]; then
+	output=$(p4 edit "$XcodeAlertAffectedPaths")
+elif [ -x "$(command -v /usr/local/bin/p4)" ]; then
+	output=$(/usr/local/bin/p4 edit "$XcodeAlertAffectedPaths")
+elif  [ -x "$(command -v /usr/bin/p4)" ]; then
+	output=$(/usr/bin/p4 edit "$XcodeAlertAffectedPaths")
+elif  [ -x "$(command -v /opt/homebrew/bin/p4)" ]; then
+	output=$(/opt/homebrew/bin/p4 edit "$XcodeAlertAffectedPaths")
+else
+	# these can be slow, so only run them as needed
+	
+	source ~/.zprofile
+	
+	if ! [ -x "$(command -v p4)" ]; then
+		#source ~/.zshrc
+	fi
+	
+	if ! [ -x "$(command -v p4)" ]; then
+		osascript -e 'set theText to "Unable to find p4 in path loaded from .zshrc and .zprofile"' -e "display dialog theText"
+		exit
+	fi
+	
+	output=$(p4 edit "$XcodeAlertAffectedPaths")
+fi
+
+# at this point we failed or output has the result
+osascript -ss - "$output"  <<EOF
+
+    on run argv -- argv is the output
+        display dialog argv
+    end run
+
+EOF
 
